Build - Attempting to fix the monolithic editor build.

diff --git a/Engine/Plugins/Tests/FbxAutomationTestBuilder/Source/FbxAutomationTestBuilder/Private/FbxAutomationBuilderModule.cpp b/Engine/Plugins/Tests/FbxAutomationTestBuilder/Source/FbxAutomationTestBuilder/Private/FbxAutomationBuilderModule.cpp
--- a/Engine/Plugins/Tests/FbxAutomationTestBuilder/Source/FbxAutomationTestBuilder/Private/FbxAutomationBuilderModule.cpp
+++ b/Engine/Plugins/Tests/FbxAutomationTestBuilder/Source/FbxAutomationTestBuilder/Private/FbxAutomationBuilderModule.cpp
@@ -84,4 +84,4 @@ TSharedRef<SDockTab> FFbxAutomationBuilderModule::MakeFbxAutomationBuilderTab(co
 	return FbxAutomationBuilderTab;
 }
 
-IMPLEMENT_MODULE(FFbxAutomationBuilderModule, FbxAutomationBuilderModule);
+IMPLEMENT_MODULE(FFbxAutomationBuilderModule, FbxAutomationTestBuilder);
